touching files

#rnx

diff --git a/Engine/Source/Editor/GameplayDebugger/Private/GameplayDebuggerCategoryConfigCustomization.cpp b/Engine/Source/Editor/GameplayDebugger/Private/GameplayDebuggerCategoryConfigCustomization.cpp
new file mode 100644
--- /dev/null
+++ b/Engine/Source/Editor/GameplayDebugger/Private/GameplayDebuggerCategoryConfigCustomization.cpp
@@ -0,0 +1,92 @@
+// Copyright Epic Games, Inc. All Rights Reserved.
+
+#include "GameplayDebuggerCategoryConfigCustomization.h"
+
+#include "UObject/UnrealType.h"
+#include "Widgets/DeclarativeSyntaxSupport.h"
+#include "Application/SlateWindowHelper.h"
+#include "Widgets/Text/STextBlock.h"
+#include "DetailWidgetRow.h"
+#include "IDetailChildrenBuilder.h"
+#include "DetailLayoutBuilder.h"
+#include "GameplayDebuggerConfig.h"
+
+#define LOCTEXT_NAMESPACE "GameplayDebuggerConfig"
+
+void FGameplayDebuggerCategoryConfigCustomization::CustomizeHeader(TSharedRef<class IPropertyHandle> StructPropertyHandle, class FDetailWidgetRow& HeaderRow, IPropertyTypeCustomizationUtils& StructCustomizationUtils)
+{
+	CategoryNameProp = StructPropertyHandle->GetChildHandle(GET_MEMBER_NAME_CHECKED(FGameplayDebuggerCategoryConfig, CategoryName));
+	SlotIdxProp = StructPropertyHandle->GetChildHandle(GET_MEMBER_NAME_CHECKED(FGameplayDebuggerCategoryConfig, SlotIdx));
+	ActiveInGameProp = StructPropertyHandle->GetChildHandle(GET_MEMBER_NAME_CHECKED(FGameplayDebuggerCategoryConfig, ActiveInGame));
+	ActiveInSimulateProp = StructPropertyHandle->GetChildHandle(GET_MEMBER_NAME_CHECKED(FGameplayDebuggerCategoryConfig, ActiveInSimulate));
+	
+	FSimpleDelegate RefreshHeaderDelegate = FSimpleDelegate::CreateSP(this, &FGameplayDebuggerCategoryConfigCustomization::OnChildValueChanged);
+	StructPropertyHandle->SetOnChildPropertyValueChanged(RefreshHeaderDelegate);
+	OnChildValueChanged();
+
+	HeaderRow.NameContent()
+	[
+		StructPropertyHandle->CreatePropertyNameWidget()
+	]
+	.ValueContent().VAlign(VAlign_Center).MinDesiredWidth(300.0f)
+	[
+		SNew(STextBlock)
+		.Text(this, &FGameplayDebuggerCategoryConfigCustomization::GetHeaderDesc)
+		.Font(IDetailLayoutBuilder::GetDetailFont())
+	];
+};
+
+void FGameplayDebuggerCategoryConfigCustomization::CustomizeChildren(TSharedRef<class IPropertyHandle> StructPropertyHandle, class IDetailChildrenBuilder& StructBuilder, IPropertyTypeCustomizationUtils& StructCustomizationUtils)
+{
+	uint32 NumChildProps = 0;
+	StructPropertyHandle->GetNumChildren(NumChildProps);
+
+	for (uint32 Idx = 0; Idx < NumChildProps; Idx++)
+	{
+		TSharedPtr<IPropertyHandle> PropHandle = StructPropertyHandle->GetChildHandle(Idx);
+		if (PropHandle->GetProperty() && PropHandle->GetProperty()->GetFName() == GET_MEMBER_NAME_CHECKED(FGameplayDebuggerCategoryConfig, CategoryName))
+		{
+			continue;
+		}
+
+		StructBuilder.AddProperty(PropHandle.ToSharedRef());
+	}
+};
+
+void FGameplayDebuggerCategoryConfigCustomization::OnChildValueChanged()
+{
+	FString CategoryNameValue;
+	if (CategoryNameProp.IsValid())
+	{
+		CategoryNameProp->GetValue(CategoryNameValue);
+	}
+
+	int32 SlotIdxValue = -1;
+	if (SlotIdxProp.IsValid())
+	{
+		SlotIdxProp->GetValue(SlotIdxValue);
+	}
+
+	uint8 ActiveInGameValue = 0;
+	if (ActiveInGameProp.IsValid())
+	{
+		ActiveInGameProp->GetValue(ActiveInGameValue);
+	}
+
+	uint8 ActiveInSimulateValue = 0;
+	if (ActiveInSimulateProp.IsValid())
+	{
+		ActiveInSimulateProp->GetValue(ActiveInSimulateValue);
+	}
+
+	CachedHeader = FText::FromString(FString::Printf(TEXT("[%s]:%s%s%s"),
+		SlotIdxValue < 0 ? TEXT("-") : *TTypeToString<int32>::ToString(SlotIdxValue),
+		CategoryNameValue.Len() ? *CategoryNameValue : TEXT("??"),
+		(ActiveInGameValue == (uint8)EGameplayDebuggerOverrideMode::UseDefault) ? TEXT("") :
+			(ActiveInGameValue == (uint8)EGameplayDebuggerOverrideMode::Enable) ? TEXT(" game:ON") : TEXT(" game:OFF"),
+		(ActiveInSimulateValue == (uint8)EGameplayDebuggerOverrideMode::UseDefault) ? TEXT("") :
+			(ActiveInSimulateValue == (uint8)EGameplayDebuggerOverrideMode::Enable) ? TEXT(" simulate:ON") : TEXT(" simulate:OFF")
+		));
+}
+
+#undef LOCTEXT_NAMESPACE
diff --git a/Engine/Source/Editor/GameplayDebugger/Private/GameplayDebuggerEdMode.cpp b/Engine/Source/Editor/GameplayDebugger/Private/GameplayDebuggerEdMode.cpp
new file mode 100644
--- /dev/null
+++ b/Engine/Source/Editor/GameplayDebugger/Private/GameplayDebuggerEdMode.cpp
@@ -0,0 +1,140 @@
+// Copyright Epic Games, Inc. All Rights Reserved.
+
+#include "GameplayDebuggerEdMode.h"
+
+#include "CoreGlobals.h"
+#include "GameplayDebuggerToolkit.h"
+#include "Toolkits/ToolkitManager.h"
+#include "GameplayDebuggerPlayerManager.h"
+#include "GameplayDebuggerLocalController.h"
+
+#include "EditorModeManager.h"
+#include "Components/InputComponent.h"
+
+const FName FGameplayDebuggerEdMode::EM_GameplayDebugger = TEXT("EM_GameplayDebugger");
+
+void FGameplayDebuggerEdMode::Enter()
+{
+	FEdMode::Enter();
+	
+	if (!Toolkit.IsValid())
+	{
+		Toolkit = MakeShareable(new FGameplayDebuggerToolkit(this));
+		Toolkit->Init(Owner->GetToolkitHost());
+	}
+
+	bPrevGAreScreenMessagesEnabled = GAreScreenMessagesEnabled;
+	GAreScreenMessagesEnabled = false;
+}
+
+void FGameplayDebuggerEdMode::Exit()
+{
+	if (Toolkit.IsValid())
+	{
+		FToolkitManager::Get().CloseToolkit(Toolkit.ToSharedRef());
+		Toolkit.Reset();
+	}
+
+	if (FocusedViewport)
+	{
+		EnableViewportClientFlags(FocusedViewport, false);
+		FocusedViewport = nullptr;
+	}
+
+	FEdMode::Exit();
+	GAreScreenMessagesEnabled = bPrevGAreScreenMessagesEnabled;
+}
+
+void FGameplayDebuggerEdMode::EnableViewportClientFlags(FEditorViewportClient* ViewportClient, bool bEnable)
+{
+	ViewportClient->bUseNumpadCameraControl = false;
+}
+
+bool FGameplayDebuggerEdMode::ReceivedFocus(FEditorViewportClient* ViewportClient, FViewport* Viewport)
+{
+	EnableViewportClientFlags(ViewportClient, true);
+	FocusedViewport = ViewportClient;
+	return false;
+}
+
+bool FGameplayDebuggerEdMode::LostFocus(FEditorViewportClient* ViewportClient, FViewport* Viewport)
+{
+	EnableViewportClientFlags(ViewportClient, false);
+	FocusedViewport = nullptr;
+	return false;
+}
+
+bool FGameplayDebuggerEdMode::InputKey(FEditorViewportClient* InViewportClient, FViewport* InViewport, FKey InKey, EInputEvent InEvent)
+{
+	UWorld* MyWorld = GetWorld();
+	APlayerController* LocalPC = GEngine->GetFirstLocalPlayerController(MyWorld);
+
+	if (LocalPC && InViewportClient)
+	{
+		// process raw input for debugger's input component manually
+		// can't use LocalPC->InputKey here, since it will trigger for every bound chord, not only gameplay debugger ones
+		// and will not work at all when simulation is paused
+
+		AGameplayDebuggerPlayerManager& PlayerManager = AGameplayDebuggerPlayerManager::GetCurrent(MyWorld);
+		const FGameplayDebuggerPlayerData* DataPtr = PlayerManager.GetPlayerData(*LocalPC);
+
+		if (DataPtr && DataPtr->InputComponent && DataPtr->Controller && DataPtr->Controller->IsKeyBound(InKey.GetFName()))
+		{
+			const FInputChord ActiveChord(InKey, InViewportClient->IsShiftPressed(), InViewportClient->IsCtrlPressed(), InViewportClient->IsAltPressed(), InViewportClient->IsCmdPressed());
+
+			// go over all bound actions
+			const int32 NumBindings = DataPtr->InputComponent->KeyBindings.Num();
+			for (int32 Idx = 0; Idx < NumBindings; Idx++)
+			{
+				const FInputKeyBinding& KeyBinding = DataPtr->InputComponent->KeyBindings[Idx];
+				if (KeyBinding.KeyEvent == InEvent && KeyBinding.Chord == ActiveChord && KeyBinding.KeyDelegate.IsBound())
+				{
+					KeyBinding.KeyDelegate.Execute(InKey);
+				}
+			}
+
+			return true;			
+		}
+	}
+
+	return false;
+}
+
+void FGameplayDebuggerEdMode::Tick(FEditorViewportClient* ViewportClient, float DeltaTime)
+{
+	if (ViewportClient == nullptr || !ViewportClient->EngineShowFlags.DebugAI)
+	{
+		Owner->DeactivateMode(FGameplayDebuggerEdMode::EM_GameplayDebugger);
+	}
+}
+
+void FGameplayDebuggerEdMode::SafeOpenMode()
+{
+	// The global mode manager should not be created or accessed in a commandlet environment
+	if (IsRunningCommandlet())
+	{
+		return;
+	}
+
+	// By calling this get, it make sure that the singleton is created and not during the SafeCloseMode whitch can be called during garbage collect.
+	GLevelEditorModeTools();
+}
+
+void FGameplayDebuggerEdMode::SafeCloseMode()
+{
+	// The global mode manager should not be created or accessed in a commandlet environment
+	if (IsRunningCommandlet())
+	{
+		return;
+	}
+
+	// this may be called on closing editor during PIE (~viewport -> teardown PIE -> debugger's cleanup on game end)
+	//
+	// DeactivateMode tries to bring up default mode, but toolkit is already destroyed by that time
+	// and editor crashes on check in GLevelEditorModeTools().GetToolkitHost() inside default mode's code
+	// note: no need to do it if IsEngineExitRequested since the GLevelEditorModeTools is no longer available 
+	if (IsEngineExitRequested() == false && GLevelEditorModeTools().HasToolkitHost())
+	{
+		GLevelEditorModeTools().DeactivateMode(FGameplayDebuggerEdMode::EM_GameplayDebugger);
+	}
+}
diff --git a/Engine/Source/Editor/GameplayDebugger/Private/GameplayDebuggerExtensionConfigCustomization.cpp b/Engine/Source/Editor/GameplayDebugger/Private/GameplayDebuggerExtensionConfigCustomization.cpp
new file mode 100644
--- /dev/null
+++ b/Engine/Source/Editor/GameplayDebugger/Private/GameplayDebuggerExtensionConfigCustomization.cpp
@@ -0,0 +1,75 @@
+// Copyright Epic Games, Inc. All Rights Reserved.
+
+#include "GameplayDebuggerExtensionConfigCustomization.h"
+
+#include "UObject/UnrealType.h"
+#include "Widgets/DeclarativeSyntaxSupport.h"
+#include "Application/SlateWindowHelper.h"
+#include "Widgets/Text/STextBlock.h"
+#include "DetailWidgetRow.h"
+#include "IDetailChildrenBuilder.h"
+#include "DetailLayoutBuilder.h"
+#include "GameplayDebuggerConfig.h"
+
+#define LOCTEXT_NAMESPACE "GameplayDebuggerConfig"
+
+void FGameplayDebuggerExtensionConfigCustomization::CustomizeHeader(TSharedRef<class IPropertyHandle> StructPropertyHandle, class FDetailWidgetRow& HeaderRow, IPropertyTypeCustomizationUtils& StructCustomizationUtils)
+{
+	ExtensionNameProp = StructPropertyHandle->GetChildHandle(GET_MEMBER_NAME_CHECKED(FGameplayDebuggerExtensionConfig, ExtensionName));
+	UseExtensionProp = StructPropertyHandle->GetChildHandle(GET_MEMBER_NAME_CHECKED(FGameplayDebuggerExtensionConfig, UseExtension));
+	
+	FSimpleDelegate RefreshHeaderDelegate = FSimpleDelegate::CreateSP(this, &FGameplayDebuggerExtensionConfigCustomization::OnChildValueChanged);
+	StructPropertyHandle->SetOnChildPropertyValueChanged(RefreshHeaderDelegate);
+	OnChildValueChanged();
+
+	HeaderRow.NameContent()
+	[
+		StructPropertyHandle->CreatePropertyNameWidget()
+	]
+	.ValueContent().VAlign(VAlign_Center).MinDesiredWidth(300.0f)
+	[
+		SNew(STextBlock)
+		.Text(this, &FGameplayDebuggerExtensionConfigCustomization::GetHeaderDesc)
+		.Font(IDetailLayoutBuilder::GetDetailFont())
+	];
+};
+
+void FGameplayDebuggerExtensionConfigCustomization::CustomizeChildren(TSharedRef<class IPropertyHandle> StructPropertyHandle, class IDetailChildrenBuilder& StructBuilder, IPropertyTypeCustomizationUtils& StructCustomizationUtils)
+{
+	uint32 NumChildProps = 0;
+	StructPropertyHandle->GetNumChildren(NumChildProps);
+
+	for (uint32 Idx = 0; Idx < NumChildProps; Idx++)
+	{
+		TSharedPtr<IPropertyHandle> PropHandle = StructPropertyHandle->GetChildHandle(Idx);
+		if (PropHandle->GetProperty() && PropHandle->GetProperty()->GetFName() == GET_MEMBER_NAME_CHECKED(FGameplayDebuggerExtensionConfig, ExtensionName))
+		{
+			continue;
+		}
+
+		StructBuilder.AddProperty(PropHandle.ToSharedRef());
+	}
+};
+
+void FGameplayDebuggerExtensionConfigCustomization::OnChildValueChanged()
+{
+	FString ExtensionNameValue;
+	if (ExtensionNameProp.IsValid())
+	{
+		ExtensionNameProp->GetValue(ExtensionNameValue);
+	}
+
+	uint8 UseExtensionValue = 0;
+	if (UseExtensionProp.IsValid())
+	{
+		UseExtensionProp->GetValue(UseExtensionValue);
+	}
+
+	CachedHeader = FText::FromString(FString::Printf(TEXT("%s %s"),
+		ExtensionNameValue.Len() ? *ExtensionNameValue : TEXT("??"),
+		(UseExtensionValue == (uint8)EGameplayDebuggerOverrideMode::UseDefault) ? TEXT("") :
+			(UseExtensionValue == (uint8)EGameplayDebuggerOverrideMode::Enable) ? TEXT("is enabled") : TEXT("is disabled")
+		));
+}
+
+#undef LOCTEXT_NAMESPACE
diff --git a/Engine/Source/Editor/GameplayDebugger/Private/GameplayDebuggerInputConfigCustomization.cpp b/Engine/Source/Editor/GameplayDebugger/Private/GameplayDebuggerInputConfigCustomization.cpp
new file mode 100644
--- /dev/null
+++ b/Engine/Source/Editor/GameplayDebugger/Private/GameplayDebuggerInputConfigCustomization.cpp
@@ -0,0 +1,109 @@
+// Copyright Epic Games, Inc. All Rights Reserved.
+
+#include "GameplayDebuggerInputConfigCustomization.h"
+#include "Widgets/DeclarativeSyntaxSupport.h"
+#include "Engine/GameViewportClient.h"
+#include "GameplayDebuggerTypes.h"
+
+#include "Widgets/Text/STextBlock.h"
+#include "DetailWidgetRow.h"
+#include "IDetailChildrenBuilder.h"
+#include "DetailLayoutBuilder.h"
+#include "GameplayDebuggerConfig.h"
+
+#define LOCTEXT_NAMESPACE "GameplayDebuggerConfig"
+
+void FGameplayDebuggerInputConfigCustomization::CustomizeHeader(TSharedRef<class IPropertyHandle> StructPropertyHandle, class FDetailWidgetRow& HeaderRow, IPropertyTypeCustomizationUtils& StructCustomizationUtils)
+{
+	ConfigNameProp = StructPropertyHandle->GetChildHandle(GET_MEMBER_NAME_CHECKED(FGameplayDebuggerInputConfig, ConfigName));
+	KeyProp = StructPropertyHandle->GetChildHandle(GET_MEMBER_NAME_CHECKED(FGameplayDebuggerInputConfig, Key));
+	ModShiftProp = StructPropertyHandle->GetChildHandle(GET_MEMBER_NAME_CHECKED(FGameplayDebuggerInputConfig, bModShift));
+	ModCtrlProp = StructPropertyHandle->GetChildHandle(GET_MEMBER_NAME_CHECKED(FGameplayDebuggerInputConfig, bModCtrl));
+	ModAltProp = StructPropertyHandle->GetChildHandle(GET_MEMBER_NAME_CHECKED(FGameplayDebuggerInputConfig, bModAlt));
+	ModCmdProp = StructPropertyHandle->GetChildHandle(GET_MEMBER_NAME_CHECKED(FGameplayDebuggerInputConfig, bModCmd));
+
+	FSimpleDelegate RefreshHeaderDelegate = FSimpleDelegate::CreateSP(this, &FGameplayDebuggerInputConfigCustomization::OnChildValueChanged);
+	StructPropertyHandle->SetOnChildPropertyValueChanged(RefreshHeaderDelegate);
+	OnChildValueChanged();
+
+	HeaderRow.NameContent()
+	[
+		StructPropertyHandle->CreatePropertyNameWidget()
+	]
+	.ValueContent().VAlign(VAlign_Center).MinDesiredWidth(300.0f)
+	[
+		SNew(STextBlock)
+		.Text(this, &FGameplayDebuggerInputConfigCustomization::GetHeaderDesc)
+		.Font(IDetailLayoutBuilder::GetDetailFont())
+	];
+};
+
+void FGameplayDebuggerInputConfigCustomization::CustomizeChildren(TSharedRef<class IPropertyHandle> StructPropertyHandle, class IDetailChildrenBuilder& StructBuilder, IPropertyTypeCustomizationUtils& StructCustomizationUtils)
+{
+	uint32 NumChildProps = 0;
+	StructPropertyHandle->GetNumChildren(NumChildProps);
+
+	for (uint32 Idx = 0; Idx < NumChildProps; Idx++)
+	{
+		TSharedPtr<IPropertyHandle> PropHandle = StructPropertyHandle->GetChildHandle(Idx);
+		if (PropHandle->GetProperty() && PropHandle->GetProperty()->GetFName() == GET_MEMBER_NAME_CHECKED(FGameplayDebuggerInputConfig, ConfigName))
+		{
+			continue;
+		}
+
+		StructBuilder.AddProperty(PropHandle.ToSharedRef());
+	}
+};
+
+void FGameplayDebuggerInputConfigCustomization::OnChildValueChanged()
+{
+	FString ConfigNameValue;
+	if (ConfigNameProp.IsValid())
+	{
+		ConfigNameProp->GetValue(ConfigNameValue);
+	}
+
+	FString KeyValue;
+	if (KeyProp.IsValid())
+	{
+		KeyProp->GetPerObjectValue(0, KeyValue);
+	}
+
+	bool bModShift = false;
+	if (ModShiftProp.IsValid())
+	{
+		ModShiftProp->GetValue(bModShift);
+	}
+
+	bool bModCtrl = false;
+	if (ModCtrlProp.IsValid())
+	{
+		ModCtrlProp->GetValue(bModCtrl);
+	}
+	
+	bool bModAlt = false;
+	if (ModAltProp.IsValid())
+	{
+		ModAltProp->GetValue(bModAlt);
+	}
+	
+	bool bModCmd = false;
+	if (ModCmdProp.IsValid())
+	{
+		ModCmdProp->GetValue(bModCmd);
+	}
+
+	FGameplayDebuggerInputHandler DummyHandler;
+	DummyHandler.KeyName = *KeyValue;
+	DummyHandler.Modifier.bShift = bModShift;
+	DummyHandler.Modifier.bCtrl = bModCtrl;
+	DummyHandler.Modifier.bAlt = bModAlt;
+	DummyHandler.Modifier.bCmd = bModCmd;
+
+	CachedHeader = FText::FromString(FString::Printf(TEXT("%s: %s"),
+		ConfigNameValue.Len() ? *ConfigNameValue : TEXT("??"),
+		*DummyHandler.ToString()
+		));
+}
+
+#undef LOCTEXT_NAMESPACE
diff --git a/Engine/Source/Editor/GameplayDebugger/Private/GameplayDebuggerToolkit.cpp b/Engine/Source/Editor/GameplayDebugger/Private/GameplayDebuggerToolkit.cpp
new file mode 100644
--- /dev/null
+++ b/Engine/Source/Editor/GameplayDebugger/Private/GameplayDebuggerToolkit.cpp
@@ -0,0 +1,106 @@
+// Copyright Epic Games, Inc. All Rights Reserved.
+
+#include "GameplayDebuggerToolkit.h"
+#include "Modules/ModuleManager.h"
+
+#include "LevelEditor.h"
+#include "IAssetViewport.h"
+#include "SLevelViewport.h"
+#include "LevelEditorViewport.h"
+#include "Widgets/Layout/SBorder.h"
+#include "Widgets/Input/SButton.h"
+#include "Widgets/Text/STextBlock.h"
+#include "Styling/AppStyle.h"
+
+#define LOCTEXT_NAMESPACE "GameplayDebuggerConfig"
+
+FGameplayDebuggerToolkit::FGameplayDebuggerToolkit(class FEdMode* InOwningMode)
+{
+	DebuggerEdMode = InOwningMode;
+}
+
+FText FGameplayDebuggerToolkit::GetBaseToolkitName() const
+{
+	return LOCTEXT("EdToolTitle", "Gameplay Debugger");
+}
+
+FName FGameplayDebuggerToolkit::GetToolkitFName() const
+{
+	return FName("GameplayDebuggerToolkit");
+}
+
+void FGameplayDebuggerToolkit::Init(const TSharedPtr<class IToolkitHost>& InitToolkitHost)
+{
+	MyWidget =
+		SNew(SBorder)
+		.BorderImage(FAppStyle::GetBrush("ToolPanel.GroupBorder"))
+		.Content()
+		[
+			SNew(SVerticalBox)
+			+ SVerticalBox::Slot()
+				.AutoHeight()
+				.HAlign(HAlign_Left)
+				.VAlign(VAlign_Center)
+				[
+					SNew(STextBlock)
+					.Text(LOCTEXT("EdToolStatus", "Gameplay debugger is active."))
+				]
+			+ SVerticalBox::Slot()
+				.AutoHeight()
+				.HAlign(HAlign_Left)
+				.VAlign(VAlign_Center)
+				.Padding(0, 30.0f)
+				[
+					SNew(STextBlock)
+					.Visibility(this, &FGameplayDebuggerToolkit::GetScreenMessageWarningVisibility)
+					.Text(LOCTEXT("EdToolMessageWarning", "Warning! On screen messages are suppressed!\nUse EnableAllScreenMessages to restore them."))
+				]
+			+ SVerticalBox::Slot()
+				.AutoHeight()
+				.HAlign(HAlign_Left)
+				.VAlign(VAlign_Center)
+				[
+					SNew(STextBlock)
+					.Text(LOCTEXT("EdToolDisableHint", "Clear DebugAI show flag to disable tool."))
+				]
+			+ SVerticalBox::Slot()
+				.AutoHeight()
+				.HAlign(HAlign_Left)
+				.VAlign(VAlign_Center)
+				.Padding(5.0f, 5.0f)
+				[
+					SNew(SButton)
+					.Text(LOCTEXT("EdToolDisableButton", "Disable tool"))
+					.OnClicked(this, &FGameplayDebuggerToolkit::OnClickedDisableTool)
+				]
+		];
+
+	FModeToolkit::Init(InitToolkitHost);
+}
+
+EVisibility FGameplayDebuggerToolkit::GetScreenMessageWarningVisibility() const
+{
+	return GAreScreenMessagesEnabled ? EVisibility::Collapsed : EVisibility::Visible;
+}
+
+FReply FGameplayDebuggerToolkit::OnClickedDisableTool()
+{
+	FLevelEditorModule& LevelEditorModule = FModuleManager::GetModuleChecked<FLevelEditorModule>("LevelEditor");
+	TSharedPtr<ILevelEditor> LevelEditor = LevelEditorModule.GetFirstLevelEditor();
+	if (LevelEditor.IsValid())
+	{
+		TArray<TSharedPtr<SLevelViewport>> Viewports = LevelEditor->GetViewports();
+		for (const TSharedPtr<SLevelViewport>& ViewportWindow : Viewports)
+		{
+			if (ViewportWindow.IsValid())
+			{
+				FEditorViewportClient& Viewport = ViewportWindow->GetAssetViewportClient();
+				Viewport.EngineShowFlags.SetDebugAI(false);
+			}
+		}
+	}
+	
+	return FReply::Handled();
+}
+
+#undef LOCTEXT_NAMESPACE
diff --git a/Engine/Source/Editor/GameplayDebugger/Public/GameplayDebuggerCategoryConfigCustomization.h b/Engine/Source/Editor/GameplayDebugger/Public/GameplayDebuggerCategoryConfigCustomization.h
new file mode 100644
--- /dev/null
+++ b/Engine/Source/Editor/GameplayDebugger/Public/GameplayDebuggerCategoryConfigCustomization.h
@@ -0,0 +1,32 @@
+// Copyright Epic Games, Inc. All Rights Reserved.
+
+#pragma once
+
+#include "CoreMinimal.h"
+
+#include "IPropertyTypeCustomization.h"
+
+class IPropertyHandle;
+
+class FGameplayDebuggerCategoryConfigCustomization : public IPropertyTypeCustomization
+{
+public:
+	static TSharedRef<IPropertyTypeCustomization> MakeInstance()
+	{
+		return MakeShareable(new FGameplayDebuggerCategoryConfigCustomization);
+	}
+
+	virtual void CustomizeHeader(TSharedRef<class IPropertyHandle> StructPropertyHandle, class FDetailWidgetRow& HeaderRow, IPropertyTypeCustomizationUtils& StructCustomizationUtils) override;
+	virtual void CustomizeChildren(TSharedRef<class IPropertyHandle> StructPropertyHandle, class IDetailChildrenBuilder& StructBuilder, IPropertyTypeCustomizationUtils& StructCustomizationUtils) override;
+
+protected:
+	TSharedPtr<IPropertyHandle> CategoryNameProp;
+	TSharedPtr<IPropertyHandle> SlotIdxProp;
+	TSharedPtr<IPropertyHandle> ActiveInGameProp;
+	TSharedPtr<IPropertyHandle> ActiveInSimulateProp;
+
+	FText CachedHeader;
+	FText GetHeaderDesc() const { return CachedHeader; }
+
+	void OnChildValueChanged();
+};
diff --git a/Engine/Source/Editor/GameplayDebugger/Public/GameplayDebuggerEdMode.h b/Engine/Source/Editor/GameplayDebugger/Public/GameplayDebuggerEdMode.h
new file mode 100644
--- /dev/null
+++ b/Engine/Source/Editor/GameplayDebugger/Public/GameplayDebuggerEdMode.h
@@ -0,0 +1,37 @@
+// Copyright Epic Games, Inc. All Rights Reserved.
+
+#pragma once
+
+#include "CoreMinimal.h"
+
+#include "InputCoreTypes.h"
+#include "EditorModes.h"
+#include "EdMode.h"
+#include "EditorViewportClient.h"
+
+class FGameplayDebuggerEdMode : public FEdMode
+{
+public:
+	FGameplayDebuggerEdMode() : FocusedViewport(nullptr) {}
+
+	// FEdMode interface
+	virtual bool UsesToolkits() const override { return true; }
+	virtual void Enter() override;
+	virtual void Exit() override;
+	virtual bool ReceivedFocus(FEditorViewportClient* ViewportClient, FViewport* Viewport) override;
+	virtual bool LostFocus(FEditorViewportClient* ViewportClient, FViewport* Viewport) override;
+	virtual bool InputKey(FEditorViewportClient* InViewportClient, FViewport* InViewport, FKey InKey, EInputEvent InEvent) override;
+	virtual void Tick(FEditorViewportClient* ViewportClient, float DeltaTime);
+	// End of FEdMode interface
+
+	static void SafeOpenMode();
+	static void SafeCloseMode();
+
+	static const FName EM_GameplayDebugger;
+
+private:
+	FEditorViewportClient* FocusedViewport;
+	bool bPrevGAreScreenMessagesEnabled;
+
+	void EnableViewportClientFlags(FEditorViewportClient* ViewportClient, bool bEnable);
+};
diff --git a/Engine/Source/Editor/GameplayDebugger/Public/GameplayDebuggerExtensionConfigCustomization.h b/Engine/Source/Editor/GameplayDebugger/Public/GameplayDebuggerExtensionConfigCustomization.h
new file mode 100644
--- /dev/null
+++ b/Engine/Source/Editor/GameplayDebugger/Public/GameplayDebuggerExtensionConfigCustomization.h
@@ -0,0 +1,30 @@
+// Copyright Epic Games, Inc. All Rights Reserved.
+
+#pragma once
+
+#include "CoreMinimal.h"
+
+#include "IPropertyTypeCustomization.h"
+
+class IPropertyHandle;
+
+class FGameplayDebuggerExtensionConfigCustomization : public IPropertyTypeCustomization
+{
+public:
+	static TSharedRef<IPropertyTypeCustomization> MakeInstance()
+	{
+		return MakeShareable(new FGameplayDebuggerExtensionConfigCustomization);
+	}
+
+	virtual void CustomizeHeader(TSharedRef<class IPropertyHandle> StructPropertyHandle, class FDetailWidgetRow& HeaderRow, IPropertyTypeCustomizationUtils& StructCustomizationUtils) override;
+	virtual void CustomizeChildren(TSharedRef<class IPropertyHandle> StructPropertyHandle, class IDetailChildrenBuilder& StructBuilder, IPropertyTypeCustomizationUtils& StructCustomizationUtils) override;
+
+protected:
+	TSharedPtr<IPropertyHandle> ExtensionNameProp;
+	TSharedPtr<IPropertyHandle> UseExtensionProp;
+
+	FText CachedHeader;
+	FText GetHeaderDesc() const { return CachedHeader; }
+
+	void OnChildValueChanged();
+};
diff --git a/Engine/Source/Editor/GameplayDebugger/Public/GameplayDebuggerInputConfigCustomization.h b/Engine/Source/Editor/GameplayDebugger/Public/GameplayDebuggerInputConfigCustomization.h
new file mode 100644
--- /dev/null
+++ b/Engine/Source/Editor/GameplayDebugger/Public/GameplayDebuggerInputConfigCustomization.h
@@ -0,0 +1,34 @@
+// Copyright Epic Games, Inc. All Rights Reserved.
+
+#pragma once
+
+#include "CoreMinimal.h"
+
+#include "IPropertyTypeCustomization.h"
+
+class IPropertyHandle;
+
+class FGameplayDebuggerInputConfigCustomization : public IPropertyTypeCustomization
+{
+public:
+	static TSharedRef<IPropertyTypeCustomization> MakeInstance()
+	{
+		return MakeShareable(new FGameplayDebuggerInputConfigCustomization);
+	}
+
+	virtual void CustomizeHeader(TSharedRef<class IPropertyHandle> StructPropertyHandle, class FDetailWidgetRow& HeaderRow, IPropertyTypeCustomizationUtils& StructCustomizationUtils) override;
+	virtual void CustomizeChildren(TSharedRef<class IPropertyHandle> StructPropertyHandle, class IDetailChildrenBuilder& StructBuilder, IPropertyTypeCustomizationUtils& StructCustomizationUtils) override;
+
+protected:
+	TSharedPtr<IPropertyHandle> ConfigNameProp;
+	TSharedPtr<IPropertyHandle> KeyProp;
+	TSharedPtr<IPropertyHandle> ModShiftProp;
+	TSharedPtr<IPropertyHandle> ModCtrlProp;
+	TSharedPtr<IPropertyHandle> ModAltProp;
+	TSharedPtr<IPropertyHandle> ModCmdProp;
+
+	FText CachedHeader;
+	FText GetHeaderDesc() const { return CachedHeader; }
+
+	void OnChildValueChanged();
+};
diff --git a/Engine/Source/Editor/GameplayDebugger/Public/GameplayDebuggerToolkit.h b/Engine/Source/Editor/GameplayDebugger/Public/GameplayDebuggerToolkit.h
new file mode 100644
--- /dev/null
+++ b/Engine/Source/Editor/GameplayDebugger/Public/GameplayDebuggerToolkit.h
@@ -0,0 +1,35 @@
+// Copyright Epic Games, Inc. All Rights Reserved.
+
+#pragma once
+
+#include "CoreMinimal.h"
+
+#include "Layout/Visibility.h"
+#include "Input/Reply.h"
+#include "Widgets/SWidget.h"
+#include "Toolkits/IToolkitHost.h"
+#include "Toolkits/BaseToolkit.h"
+
+class FGameplayDebuggerToolkit : public FModeToolkit
+{
+public:
+	FGameplayDebuggerToolkit(class FEdMode* InOwningMode);
+
+	// IToolkit interface
+	virtual FText GetBaseToolkitName() const override;
+	virtual FName GetToolkitFName() const override;
+	virtual class FEdMode* GetEditorMode() const override { return DebuggerEdMode; }
+	virtual TSharedPtr<class SWidget> GetInlineContent() const override { return MyWidget; }
+	// End of IToolkit interface
+
+	// FModeToolkit interface
+	virtual void Init(const TSharedPtr<class IToolkitHost>& InitToolkitHost) override;
+	// End of FModeToolkit interface
+
+private:
+	class FEdMode* DebuggerEdMode;
+	TSharedPtr<class SWidget> MyWidget;
+
+	EVisibility GetScreenMessageWarningVisibility() const;
+	FReply OnClickedDisableTool();
+};
