[SceneGraph] Added an extra (redundant) nullptr check to make the static analyzer happy

#Synced-CL 51022704

diff --git a/Engine/Source/Editor/GameplayDebugger/Private/GameplayDebuggerEditorModule.cpp b/Engine/Source/Editor/GameplayDebugger/Private/GameplayDebuggerEditorModule.cpp
--- a/Engine/Source/Editor/GameplayDebugger/Private/GameplayDebuggerEditorModule.cpp
+++ b/Engine/Source/Editor/GameplayDebugger/Private/GameplayDebuggerEditorModule.cpp
@@ -7,10 +7,10 @@
 #include "EditorModeRegistry.h"
 #include "Engine/World.h"
 #include "GameplayDebugger.h"
-#include "GameplayDebuggerCategoryConfigCustomization.h"
-#include "GameplayDebuggerEdMode.h"
-#include "GameplayDebuggerExtensionConfigCustomization.h"
-#include "GameplayDebuggerInputConfigCustomization.h"
+//#include "GameplayDebuggerCategoryConfigCustomization.h"
+//#include "GameplayDebuggerEdMode.h"
+//#include "GameplayDebuggerExtensionConfigCustomization.h"
+//#include "GameplayDebuggerInputConfigCustomization.h"
 #include "GameplayDebuggerModule.h"
 #include "PropertyEditorModule.h"
 
@@ -23,11 +23,11 @@ void FGameplayDebuggerEditorModule::StartupModule()
 	FGameplayDebuggerModule::OnDebuggerEdModeActivation.AddRaw(this, &FGameplayDebuggerEditorModule::OnDebuggerEdModeActivation);
 
 	FPropertyEditorModule& PropertyEditorModule = FModuleManager::GetModuleChecked<FPropertyEditorModule>("PropertyEditor");
-	PropertyEditorModule.RegisterCustomPropertyTypeLayout("GameplayDebuggerCategoryConfig", FOnGetPropertyTypeCustomizationInstance::CreateStatic(&FGameplayDebuggerCategoryConfigCustomization::MakeInstance));
+	/*PropertyEditorModule.RegisterCustomPropertyTypeLayout("GameplayDebuggerCategoryConfig", FOnGetPropertyTypeCustomizationInstance::CreateStatic(&FGameplayDebuggerCategoryConfigCustomization::MakeInstance));
 	PropertyEditorModule.RegisterCustomPropertyTypeLayout("GameplayDebuggerExtensionConfig", FOnGetPropertyTypeCustomizationInstance::CreateStatic(&FGameplayDebuggerExtensionConfigCustomization::MakeInstance));
 	PropertyEditorModule.RegisterCustomPropertyTypeLayout("GameplayDebuggerInputConfig", FOnGetPropertyTypeCustomizationInstance::CreateStatic(&FGameplayDebuggerInputConfigCustomization::MakeInstance));
 
-	FEditorModeRegistry::Get().RegisterMode<FGameplayDebuggerEdMode>(FGameplayDebuggerEdMode::EM_GameplayDebugger);
+	FEditorModeRegistry::Get().RegisterMode<FGameplayDebuggerEdMode>(FGameplayDebuggerEdMode::EM_GameplayDebugger);*/
 }
 
 void FGameplayDebuggerEditorModule::ShutdownModule()
@@ -40,20 +40,20 @@ void FGameplayDebuggerEditorModule::ShutdownModule()
 		PropertyEditorModule->UnregisterCustomPropertyTypeLayout("GameplayDebuggerInputConfig");
 	}
 
-	FEditorModeRegistry::Get().UnregisterMode(FGameplayDebuggerEdMode::EM_GameplayDebugger);
+	//FEditorModeRegistry::Get().UnregisterMode(FGameplayDebuggerEdMode::EM_GameplayDebugger);
 }
 
 void FGameplayDebuggerEditorModule::OnDebuggerEdModeActivation()
 {
-	GLevelEditorModeTools().ActivateMode(FGameplayDebuggerEdMode::EM_GameplayDebugger);
+	//GLevelEditorModeTools().ActivateMode(FGameplayDebuggerEdMode::EM_GameplayDebugger);
 }
 
 void FGameplayDebuggerEditorModule::OnLocalControllerInitialized()
 {
-	FGameplayDebuggerEdMode::SafeOpenMode();
+	//FGameplayDebuggerEdMode::SafeOpenMode();
 }
 
 void FGameplayDebuggerEditorModule::OnLocalControllerUninitialized()
 {
-	FGameplayDebuggerEdMode::SafeCloseMode();
+	//FGameplayDebuggerEdMode::SafeCloseMode();
 }
